Curve refactor fixes: Pose Assets, Pose Driver and more

Addresses curves all being set to NAME_None for pose assets saved between CLs 3002109 and 3026802.
Fixes curve extraction in pose assets.
Fixes pose driver not setting curves unless they already exist.
Fixes instances where TBaseBlendedCurve::Set was being used incorrectly - this now inserts the element if it wasnt present before, potentially causing a re-sort but at least not behaving un unexpected ways compared to the old use case.

diff --git a/Engine/Plugins/Animation/ControlRig/Source/ControlRig/Private/ControlRigAnimInstance.cpp b/Engine/Plugins/Animation/ControlRig/Source/ControlRig/Private/ControlRigAnimInstance.cpp
--- a/Engine/Plugins/Animation/ControlRig/Source/ControlRig/Private/ControlRigAnimInstance.cpp
+++ b/Engine/Plugins/Animation/ControlRig/Source/ControlRig/Private/ControlRigAnimInstance.cpp
@@ -2,6 +2,7 @@
 
 #include "ControlRigAnimInstance.h"
 
+#include "Animation/AnimCurveUtils.h"
 #include "Animation/AnimNodeBase.h"
 
 #include UE_INLINE_GENERATED_CPP_BY_NAME(ControlRigAnimInstance)
@@ -50,10 +51,9 @@ bool FControlRigAnimInstanceProxy::Evaluate(FPoseContext& Output)
 		Output.Pose[ModifiedBone] = CompactPose[ModifiedBone];
 	}
 
-	for (const TPair<FName, float>& Pair : StoredCurves)
-	{
-		Output.Curve.Set(Pair.Key, Pair.Value);
-	}
+	FBlendedCurve Curve;
+	UE::Anim::FCurveUtils::BuildUnsorted(Curve, StoredCurves);
+	Output.Curve.Combine(Curve);
 
 	return true;
 }
diff --git a/Engine/Source/Runtime/AnimGraphRuntime/Private/AnimNodes/AnimNode_PoseBlendNode.cpp b/Engine/Source/Runtime/AnimGraphRuntime/Private/AnimNodes/AnimNode_PoseBlendNode.cpp
--- a/Engine/Source/Runtime/AnimGraphRuntime/Private/AnimNodes/AnimNode_PoseBlendNode.cpp
+++ b/Engine/Source/Runtime/AnimGraphRuntime/Private/AnimNodes/AnimNode_PoseBlendNode.cpp
@@ -20,14 +20,7 @@ void FAnimNode_PoseBlendNode::Initialize_AnyThread(const FAnimationInitializeCon
 {
 	DECLARE_SCOPE_HIERARCHICAL_COUNTER_ANIMNODE(Initialize_AnyThread)
 	FAnimNode_PoseHandler::Initialize_AnyThread(Context);
-
-	BulkCurves.Empty();
-	for (int32 PoseIdx = 0; PoseIdx < PoseExtractContext.PoseCurves.Num(); ++PoseIdx)
-	{
-		FPoseCurve& PoseCurve = PoseExtractContext.PoseCurves[PoseIdx];
-		BulkCurves.Add(PoseCurve.Name, BulkCurves.Num());
-	}
-
+	
 	SourcePose.Initialize(Context);
 }
 
@@ -105,4 +98,14 @@ void FAnimNode_PoseBlendNode::GatherDebugData(FNodeDebugData& DebugData)
 	SourcePose.GatherDebugData(DebugData.BranchFlow(1.f));
 }
 
-
+void FAnimNode_PoseBlendNode::RebuildPoseList(const FBoneContainer& InBoneContainer, const UPoseAsset* InPoseAsset)
+{
+	FAnimNode_PoseHandler::RebuildPoseList(InBoneContainer, InPoseAsset);
+	
+	BulkCurves.Empty();
+	for (int32 PoseIdx = 0; PoseIdx < PoseExtractContext.PoseCurves.Num(); ++PoseIdx)
+	{
+		FPoseCurve& PoseCurve = PoseExtractContext.PoseCurves[PoseIdx];
+		BulkCurves.Add(PoseCurve.Name, BulkCurves.Num());
+	}
+}
\ No newline at end of file
diff --git a/Engine/Source/Runtime/AnimGraphRuntime/Private/AnimNodes/AnimNode_PoseDriver.cpp b/Engine/Source/Runtime/AnimGraphRuntime/Private/AnimNodes/AnimNode_PoseDriver.cpp
--- a/Engine/Source/Runtime/AnimGraphRuntime/Private/AnimNodes/AnimNode_PoseDriver.cpp
+++ b/Engine/Source/Runtime/AnimGraphRuntime/Private/AnimNodes/AnimNode_PoseDriver.cpp
@@ -2,6 +2,7 @@
 
 #include "AnimNodes/AnimNode_PoseDriver.h"
 #include "AnimationRuntime.h"
+#include "Animation/AnimCurveUtils.h"
 #include "Serialization/CustomVersion.h"
 #include "Animation/AnimInstanceProxy.h"
 #include "RBF/RBFSolver.h"
@@ -433,15 +434,26 @@ void FAnimNode_PoseDriver::Evaluate_AnyThread(FPoseContext& Output)
 				Output = SourceData;
 
 				// Then set curves based on target weights
+				FBlendedCurve DrivenCurves;
 				for (const FRBFOutputWeight& Weight : OutputWeights)
 				{
 					FPoseDriverTarget& PoseTarget = PoseTargets[Weight.TargetIndex];
 					if (PoseTarget.DrivenName != NAME_None)
 					{
-						Output.Curve.Set(PoseTarget.DrivenName, Weight.TargetWeight);
+						DrivenCurves.Add(PoseTarget.DrivenName, Weight.TargetWeight);
 					}
 				}
 
+				// Merge driven curves into output
+				UE::Anim::FNamedValueArrayUtils::Union(Output.Curve, DrivenCurves, [](UE::Anim::FCurveElement& InOutElement, const UE::Anim::FCurveElement& InElement0, UE::Anim::ENamedValueUnionFlags InFlags)
+				{
+					if(EnumHasAnyFlags(InFlags, UE::Anim::ENamedValueUnionFlags::ValidArg1))
+					{
+						InOutElement.Value = InElement0.Value;
+						InOutElement.Flags |= InElement0.Flags;
+					}
+				});
+
 				bHaveValidPose = true;
 			}
 		}
diff --git a/Engine/Source/Runtime/AnimGraphRuntime/Public/AnimNodes/AnimNode_PoseBlendNode.h b/Engine/Source/Runtime/AnimGraphRuntime/Public/AnimNodes/AnimNode_PoseBlendNode.h
--- a/Engine/Source/Runtime/AnimGraphRuntime/Public/AnimNodes/AnimNode_PoseBlendNode.h
+++ b/Engine/Source/Runtime/AnimGraphRuntime/Public/AnimNodes/AnimNode_PoseBlendNode.h
@@ -43,5 +43,7 @@ public:
 	virtual void GatherDebugData(FNodeDebugData& DebugData) override;
 	// End of FAnimNode_Base interface
 
+	// FAnimNode_PoseHandler interface 
+	virtual void RebuildPoseList(const FBoneContainer& InBoneContainer, const UPoseAsset* InPoseAsset) override;
 };
 
diff --git a/Engine/Source/Runtime/Engine/Private/Animation/AnimCurveTypes.cpp b/Engine/Source/Runtime/Engine/Private/Animation/AnimCurveTypes.cpp
--- a/Engine/Source/Runtime/Engine/Private/Animation/AnimCurveTypes.cpp
+++ b/Engine/Source/Runtime/Engine/Private/Animation/AnimCurveTypes.cpp
@@ -61,7 +61,13 @@ void FAnimCurveBase::PostSerialize(const FArchive& Ar)
 	{
 		if (Ar.CustomVer(FFrameworkObjectVersion::GUID) < FFrameworkObjectVersion::SmartNameRefactor)
 		{
-			Name_DEPRECATED.DisplayName = LastObservedName_DEPRECATED;
+			// Between CLs 3002109 (SmartNameRefactor) and 3026802 (PoseAssetSupportPerBoneMask), PoseAssets had no custom version serialized into their archive, so we
+			// cant properly upgrade curves stored there. We instead assume that if LastObservedName_DEPRECATED is not
+			// NAME_None, we can use it
+			if(LastObservedName_DEPRECATED != NAME_None)
+			{
+				Name_DEPRECATED.DisplayName = LastObservedName_DEPRECATED;
+			}
 		}
 		
 		if(Ar.CustomVer(FAnimObjectVersion::GUID) < FAnimObjectVersion::AnimSequenceCurveColors)
diff --git a/Engine/Source/Runtime/Engine/Private/Animation/AnimSingleNodeInstanceProxy.cpp b/Engine/Source/Runtime/Engine/Private/Animation/AnimSingleNodeInstanceProxy.cpp
--- a/Engine/Source/Runtime/Engine/Private/Animation/AnimSingleNodeInstanceProxy.cpp
+++ b/Engine/Source/Runtime/Engine/Private/Animation/AnimSingleNodeInstanceProxy.cpp
@@ -9,6 +9,7 @@
 #include "Animation/AnimStreamable.h"
 #include "Animation/AnimSingleNodeInstance.h"
 #include "Animation/AnimationPoseData.h"
+#include "Animation/AnimCurveUtils.h"
 #include "Animation/AnimSyncScope.h"
 
 #include UE_INLINE_GENERATED_CPP_BY_NAME(AnimSingleNodeInstanceProxy)
@@ -62,10 +63,10 @@ void FAnimSingleNodeInstanceProxy::PropagatePreviewCurve(FPoseContext& Output)
 {
 	USkeleton* MySkeleton = GetSkeleton();
 
-	for (auto Iter = PreviewCurveOverride.CreateConstIterator(); Iter; ++Iter)
-	{
-		Output.Curve.Set(Iter.Key(), Iter.Value());
-	}
+	FBlendedCurve Curve;
+	UE::Anim::FCurveUtils::BuildUnsorted(Curve, PreviewCurveOverride);
+
+	Output.Curve.Combine(Curve);
 }
 #endif // WITH_EDITORONLY_DATA
 
diff --git a/Engine/Source/Runtime/Engine/Private/Animation/PoseAsset.cpp b/Engine/Source/Runtime/Engine/Private/Animation/PoseAsset.cpp
--- a/Engine/Source/Runtime/Engine/Private/Animation/PoseAsset.cpp
+++ b/Engine/Source/Runtime/Engine/Private/Animation/PoseAsset.cpp
@@ -102,7 +102,7 @@ void FPoseDataContainer::BlendPoseCurve(const FPoseData* PoseData, FBlendedCurve
 		FBlendedCurve Curve;
 		Curve.SetFilter(InOutCurve.GetFilter());
 
-		GetPoseCurve(PoseData, InOutCurve);
+		GetPoseCurve(PoseData, Curve);
 
 		UE::Anim::FNamedValueArrayUtils::Union(InOutCurve, Curve,
 			[Weight](UE::Anim::FCurveElement& InOutElement, const UE::Anim::FCurveElement& InElement, UE::Anim::ENamedValueUnionFlags InFlags)
diff --git a/Engine/Source/Runtime/Engine/Public/Animation/AnimCurveTypes.h b/Engine/Source/Runtime/Engine/Public/Animation/AnimCurveTypes.h
--- a/Engine/Source/Runtime/Engine/Public/Animation/AnimCurveTypes.h
+++ b/Engine/Source/Runtime/Engine/Public/Animation/AnimCurveTypes.h
@@ -569,8 +569,9 @@ public:
 	}
 
 	/**
-	 * Set value of curve named InName to InValue. Has no effect if the curve does not contain an element for InName.
-	 * Note that this performs a binary search per-call. Consider using a combiner operation to set multiple element's values.
+	 * Set value of curve named InName to InValue. Inserts a new value if an element for InName does not already exist.
+	 * Note that this performs a binary search per-call and can potentially cause a re-sort. Consider using a combiner
+	 * operation to set multiple element's values.
 	 * @param	InName	the name of the curve to set
 	 * @param	InValue	the value of the curve to set
 	 */
@@ -580,11 +581,17 @@ public:
 		{
 			CurveElement->Value = InValue;
 		}
+		else
+		{
+			Super::Add(InName, InValue);
+		}
 	}
 
 	/**
-	 * Set flags of curve named InName to InValue. Has no effect if the curve does not contain an element for InName.
-	 * Note that this performs a binary search per-call. Consider using a combiner operation to set multiple element's flags.
+	 * Set flags of curve named InName to InFlags.  Inserts a new default value if an element for InName does not
+	 * already exist and then sets its flags to InFlags.
+	 * Note that this performs a binary search per-call and can potentially cause a re-sort. Consider using a combiner
+	 * operation to set multiple element's flags.
 	 * @param	InName	the name of the curve to set
 	 * @param	InFlags	the flags of the curve to set
 	 */
@@ -594,6 +601,10 @@ public:
 		{
 			CurveElement->Flags = InFlags;
 		}
+		else
+		{
+			Super::Add(InName, InFlags);
+		}
 	}
 
 	UE_DEPRECATED(5.3, "Please use Get with a curve name or consider using one of the bulk APIs (e.g. UE::Anim::FCurveUtils::BulkGet)")
diff --git a/Engine/Source/Runtime/LiveLinkAnimationCore/Private/LiveLinkRetargetAsset.cpp b/Engine/Source/Runtime/LiveLinkAnimationCore/Private/LiveLinkRetargetAsset.cpp
--- a/Engine/Source/Runtime/LiveLinkAnimationCore/Private/LiveLinkRetargetAsset.cpp
+++ b/Engine/Source/Runtime/LiveLinkAnimationCore/Private/LiveLinkRetargetAsset.cpp
@@ -6,6 +6,7 @@
 #include "Animation/AnimTypes.h"
 #include "Animation/Skeleton.h"
 #include "BonePose.h"
+#include "Animation/AnimCurveUtils.h"
 #include "GenericPlatform/GenericPlatformMath.h"
 #include "Roles/LiveLinkAnimationTypes.h"
 
@@ -27,23 +28,33 @@ void ULiveLinkRetargetAsset::BuildCurveData(const FLiveLinkSkeletonStaticData* I
 
 	if (InSkeletonData->PropertyNames.Num() == InFrameData->PropertyValues.Num())
 	{
-		for (int32 CurveIdx = 0; CurveIdx < InSkeletonData->PropertyNames.Num(); ++CurveIdx)
+		auto GetCurveNameFromIndex = [InSkeletonData](int32 InCurveIndex)
 		{
-			const float Curve = InFrameData->PropertyValues[CurveIdx];
+			return InSkeletonData->PropertyNames[InCurveIndex];
+		};
+
+		auto GetCurveValueFromIndex = [InFrameData](int32 InCurveIndex)
+		{
+			const float Curve = InFrameData->PropertyValues[InCurveIndex];
 			if (FMath::IsFinite(Curve))
 			{
-				ApplyCurveValue(Skeleton, InSkeletonData->PropertyNames[CurveIdx], Curve, OutCurve);
+				return InFrameData->PropertyValues[InCurveIndex];
 			}
-		}
+			else
+			{
+				return 0.0f;
+			}
+		};
+
+		FBlendedCurve Curve;
+		UE::Anim::FCurveUtils::BuildUnsorted(Curve, InSkeletonData->PropertyNames.Num(), GetCurveNameFromIndex, GetCurveValueFromIndex);
+		OutCurve.Combine(Curve);
 	}
 }
 
 void ULiveLinkRetargetAsset::BuildCurveData(const TMap<FName, float>& CurveMap, const FCompactPose& InPose, FBlendedCurve& OutCurve) const
 {
-	const USkeleton* Skeleton = InPose.GetBoneContainer().GetSkeletonAsset();
-
-	for (const TPair<FName, float>& Pair : CurveMap)
-	{
-		ApplyCurveValue(Skeleton, Pair.Key, Pair.Value, OutCurve);
-	}
+	FBlendedCurve Curve;
+	UE::Anim::FCurveUtils::BuildUnsorted(Curve, CurveMap);
+	OutCurve.Combine(Curve);
 }
