- Adding caching of livelink data to enable/disable evaluation in livelink_pose node

diff --git a/Engine/Source/Runtime/LiveLinkAnimationCore/Private/AnimNode_LiveLinkPose.cpp b/Engine/Source/Runtime/LiveLinkAnimationCore/Private/AnimNode_LiveLinkPose.cpp
--- a/Engine/Source/Runtime/LiveLinkAnimationCore/Private/AnimNode_LiveLinkPose.cpp
+++ b/Engine/Source/Runtime/LiveLinkAnimationCore/Private/AnimNode_LiveLinkPose.cpp
@@ -28,9 +28,37 @@ void FAnimNode_LiveLinkPose::OnInitializeAnimInstance(const FAnimInstanceProxy*
 	Super::OnInitializeAnimInstance(InProxy, InAnimInstance);
 }
 
+void FAnimNode_LiveLinkPose::BuildPoseFromAnimData(const FLiveLinkSubjectFrameData& LiveLinkData, FPoseContext& Output)
+{
+	const FLiveLinkSkeletonStaticData* SkeletonData = LiveLinkData.StaticData.Cast<FLiveLinkSkeletonStaticData>();
+	const FLiveLinkAnimationFrameData* FrameData = LiveLinkData.FrameData.Cast<FLiveLinkAnimationFrameData>();
+	check(SkeletonData);
+	check(FrameData);
+
+	CurrentRetargetAsset->BuildPoseFromAnimationData(CachedDeltaTime, SkeletonData, FrameData, Output.Pose);
+	CurrentRetargetAsset->BuildPoseAndCurveFromBaseData(CachedDeltaTime, SkeletonData, FrameData, Output.Pose, Output.Curve);
+	CachedDeltaTime = 0.f; // Reset so that if we evaluate again we don't "create" time inside of the retargeter
+}
+
+void FAnimNode_LiveLinkPose::BuildPoseFromCurveData(const FLiveLinkSubjectFrameData& LiveLinkData, FPoseContext& Output)
+{
+	const FLiveLinkBaseStaticData* BaseStaticData = LiveLinkData.StaticData.Cast<FLiveLinkBaseStaticData>();
+	const FLiveLinkBaseFrameData* BaseFrameData = LiveLinkData.FrameData.Cast<FLiveLinkBaseFrameData>();
+	check(BaseStaticData);
+	check(BaseFrameData);
+
+	CurrentRetargetAsset->BuildPoseAndCurveFromBaseData(CachedDeltaTime, BaseStaticData, BaseFrameData, Output.Pose, Output.Curve);
+	CachedDeltaTime = 0.f; // Reset so that if we evaluate again we don't "create" time inside of the retargeter
+}
+
 void FAnimNode_LiveLinkPose::Initialize_AnyThread(const FAnimationInitializeContext& Context)
 {
 	InputPose.Initialize(Context);
+
+	if (CachedLiveLinkData.IsValid() == false)
+	{
+		CachedLiveLinkData = MakeShared<FLiveLinkSubjectFrameData>();
+	}
 }
 
 void FAnimNode_LiveLinkPose::PreUpdate(const UAnimInstance* InAnimInstance)
@@ -81,36 +109,51 @@ void FAnimNode_LiveLinkPose::Evaluate_AnyThread(FPoseContext& Output)
 
 	FLiveLinkSubjectFrameData SubjectFrameData;
 
-	TSubclassOf<ULiveLinkRole> SubjectRole = LiveLinkClient_AnyThread->GetSubjectRole_AnyThread(LiveLinkSubjectName);
-	if (SubjectRole)
+	if (bDoLiveLinkEvaluation)
 	{
-		if (LiveLinkClient_AnyThread->DoesSubjectSupportsRole_AnyThread(LiveLinkSubjectName, ULiveLinkAnimationRole::StaticClass()))
+		// Invalidate cached evaluated Role to make sure we have a valid one during the last evaluation when using it
+		CachedEvaluatedRole = nullptr;
+
+		TSubclassOf<ULiveLinkRole> SubjectRole = LiveLinkClient_AnyThread->GetSubjectRole_AnyThread(LiveLinkSubjectName);
+		if (SubjectRole)
 		{
-			//Process animation data if the subject is from that type
-			if (LiveLinkClient_AnyThread->EvaluateFrame_AnyThread(LiveLinkSubjectName, ULiveLinkAnimationRole::StaticClass(), SubjectFrameData))
+			if (LiveLinkClient_AnyThread->DoesSubjectSupportsRole_AnyThread(LiveLinkSubjectName, ULiveLinkAnimationRole::StaticClass()))
 			{
-				FLiveLinkSkeletonStaticData* SkeletonData = SubjectFrameData.StaticData.Cast<FLiveLinkSkeletonStaticData>();
-				FLiveLinkAnimationFrameData* FrameData = SubjectFrameData.FrameData.Cast<FLiveLinkAnimationFrameData>();
-				check(SkeletonData);
-				check(FrameData);
-
-				CurrentRetargetAsset->BuildPoseFromAnimationData(CachedDeltaTime, SkeletonData, FrameData, Output.Pose);
-				CurrentRetargetAsset->BuildPoseAndCurveFromBaseData(CachedDeltaTime, SkeletonData, FrameData, Output.Pose, Output.Curve);
-				CachedDeltaTime = 0.f; // Reset so that if we evaluate again we don't "create" time inside of the retargeter
+				//Process animation data if the subject is from that type
+				if (LiveLinkClient_AnyThread->EvaluateFrame_AnyThread(LiveLinkSubjectName, ULiveLinkAnimationRole::StaticClass(), SubjectFrameData))
+				{
+					BuildPoseFromAnimData(SubjectFrameData, Output);
+
+					CachedEvaluatedRole = ULiveLinkAnimationRole::StaticClass();
+				}
 			}
+			else if (LiveLinkClient_AnyThread->DoesSubjectSupportsRole_AnyThread(LiveLinkSubjectName, ULiveLinkBasicRole::StaticClass()))
+			{
+				//Otherwise, fetch basic data that contains property / curve data
+				if (LiveLinkClient_AnyThread->EvaluateFrame_AnyThread(LiveLinkSubjectName, ULiveLinkBasicRole::StaticClass(), SubjectFrameData))
+				{
+					BuildPoseFromCurveData(SubjectFrameData, Output);
+
+					CachedEvaluatedRole = ULiveLinkBasicRole::StaticClass();
+				}
+			}
+
+			// Keep a cached version of evaluated data to continue building pose with it when we pause evaluation
+			CachedLiveLinkData->StaticData = MoveTemp(SubjectFrameData.StaticData);
+			CachedLiveLinkData->FrameData = MoveTemp(SubjectFrameData.FrameData);
 		}
-		else if (LiveLinkClient_AnyThread->DoesSubjectSupportsRole_AnyThread(LiveLinkSubjectName, ULiveLinkBasicRole::StaticClass()))
+	}
+	else
+	{
+		if(CachedLiveLinkData && CachedEvaluatedRole)
 		{
-			//Otherwise, fetch basic data that contains property / curve data
-			if (LiveLinkClient_AnyThread->EvaluateFrame_AnyThread(LiveLinkSubjectName, ULiveLinkBasicRole::StaticClass(), SubjectFrameData))
+			if (CachedEvaluatedRole == ULiveLinkAnimationRole::StaticClass())
 			{
-				FLiveLinkBaseStaticData* BaseStaticData = SubjectFrameData.StaticData.Cast<FLiveLinkBaseStaticData>();
-				FLiveLinkBaseFrameData* BaseFrameData = SubjectFrameData.FrameData.Cast<FLiveLinkBaseFrameData>();
-				check(BaseStaticData);
-				check(BaseFrameData);
-
-				CurrentRetargetAsset->BuildPoseAndCurveFromBaseData(CachedDeltaTime, BaseStaticData, BaseFrameData, Output.Pose, Output.Curve);
-				CachedDeltaTime = 0.f; // Reset so that if we evaluate again we don't "create" time inside of the retargeter
+				BuildPoseFromAnimData(*CachedLiveLinkData, Output);
+			}
+			else if (CachedEvaluatedRole == ULiveLinkBasicRole::StaticClass())
+			{
+				BuildPoseFromCurveData(*CachedLiveLinkData, Output);
 			}
 		}
 	}
diff --git a/Engine/Source/Runtime/LiveLinkAnimationCore/Public/AnimNode_LiveLinkPose.h b/Engine/Source/Runtime/LiveLinkAnimationCore/Public/AnimNode_LiveLinkPose.h
--- a/Engine/Source/Runtime/LiveLinkAnimationCore/Public/AnimNode_LiveLinkPose.h
+++ b/Engine/Source/Runtime/LiveLinkAnimationCore/Public/AnimNode_LiveLinkPose.h
@@ -7,11 +7,14 @@
 #include "CoreMinimal.h"
 #include "LiveLinkRetargetAsset.h"
 #include "LiveLinkTypes.h"
+#include "Templates/SubclassOf.h"
 
 #include "AnimNode_LiveLinkPose.generated.h"
 
 
 class ILiveLinkClient;
+class ULiveLinkRole;
+struct FLiveLinkSubjectFrameData;
 
 
 PRAGMA_DISABLE_DEPRECATION_WARNINGS
@@ -28,6 +31,9 @@ public:
 	UPROPERTY(EditAnywhere, BlueprintReadWrite, Category = SourceData, meta = (PinShownByDefault))
 	FLiveLinkSubjectName LiveLinkSubjectName;
 
+	UPROPERTY(EditAnywhere, BlueprintReadWrite, Transient, Category = SourceData, meta = (PinShownByDefault))
+	bool bDoLiveLinkEvaluation = true;
+
 #if WITH_EDITORONLY_DATA
 	UE_DEPRECATED(4.23, "FName SubjectName is deprecated. Use the SubjectName of type FLiveLinkSubjectName instead.")
 	UPROPERTY()
@@ -57,6 +63,13 @@ public:
 
 protected:
 	LIVELINKANIMATIONCORE_API virtual void OnInitializeAnimInstance(const FAnimInstanceProxy* InProxy, const UAnimInstance* InAnimInstance) override;
+	
+	/** Build output pose using evaluated LiveLink animation data */
+	LIVELINKANIMATIONCORE_API void BuildPoseFromAnimData(const FLiveLinkSubjectFrameData& LiveLinkData, FPoseContext& Output);
+	
+	/** Build output pose using evaluated LiveLink curve (Basic Role) data */
+	LIVELINKANIMATIONCORE_API void BuildPoseFromCurveData(const FLiveLinkSubjectFrameData& LiveLinkData, FPoseContext& Output);
+
 
 private:
 
@@ -64,6 +77,12 @@ private:
 
 	// Delta time from update so that it can be passed to retargeter
 	float CachedDeltaTime;
+
+	/** Cached LiveLink evaluated data from last frame */
+	TSharedPtr<FLiveLinkSubjectFrameData> CachedLiveLinkData;
+
+	/** Cached LiveLink Role evaluated last frame */
+	TSubclassOf<ULiveLinkRole> CachedEvaluatedRole;
 };
 
 PRAGMA_ENABLE_DEPRECATION_WARNINGS
diff --git a/Engine/Source/Runtime/LiveLinkAnimationCore/Public/LiveLinkInstance.h b/Engine/Source/Runtime/LiveLinkAnimationCore/Public/LiveLinkInstance.h
--- a/Engine/Source/Runtime/LiveLinkAnimationCore/Public/LiveLinkInstance.h
+++ b/Engine/Source/Runtime/LiveLinkAnimationCore/Public/LiveLinkInstance.h
@@ -56,6 +56,17 @@ public:
 		GetProxyOnGameThread<FLiveLinkInstanceProxy>().PoseNode.RetargetAsset = RetargetAsset;
 	}
 
+	UFUNCTION(BlueprintCallable, Category = "Animation|Live Link")
+	void EnableLiveLinkEvaluation(bool bDoEnable)
+	{
+		GetProxyOnGameThread<FLiveLinkInstanceProxy>().PoseNode.bDoLiveLinkEvaluation = bDoEnable;
+	}
+
+	bool GetEnableLiveLinkEvaluation()
+	{
+		return GetProxyOnGameThread<FLiveLinkInstanceProxy>().PoseNode.bDoLiveLinkEvaluation;
+	}
+
 protected:
 	LIVELINKANIMATIONCORE_API virtual FAnimInstanceProxy* CreateAnimInstanceProxy() override;
 	LIVELINKANIMATIONCORE_API virtual void DestroyAnimInstanceProxy(FAnimInstanceProxy* InProxy) override;
