LiveLink - Fix crash when evaluating livelink pose with invalid frame data

#rnx

#Synced-CL 50901149

diff --git a/Engine/Source/Runtime/LiveLinkAnimationCore/Private/AnimNode_LiveLinkPose.cpp b/Engine/Source/Runtime/LiveLinkAnimationCore/Private/AnimNode_LiveLinkPose.cpp
--- a/Engine/Source/Runtime/LiveLinkAnimationCore/Private/AnimNode_LiveLinkPose.cpp
+++ b/Engine/Source/Runtime/LiveLinkAnimationCore/Private/AnimNode_LiveLinkPose.cpp
@@ -36,8 +36,12 @@ void FAnimNode_LiveLinkPose::BuildPoseFromAnimData(const FLiveLinkSubjectFrameDa
 	check(SkeletonData);
 	check(FrameData);
 
-	CurrentRetargetAsset->BuildPoseFromAnimationData(CachedDeltaTime, SkeletonData, FrameData, Output.Pose);
-	CurrentRetargetAsset->BuildPoseAndCurveFromBaseData(CachedDeltaTime, SkeletonData, FrameData, Output.Pose, Output.Curve);
+	if (FrameData->Transforms.Num() || FrameData->PropertyValues.Num())
+	{
+		CurrentRetargetAsset->BuildPoseFromAnimationData(CachedDeltaTime, SkeletonData, FrameData, Output.Pose);
+		CurrentRetargetAsset->BuildPoseAndCurveFromBaseData(CachedDeltaTime, SkeletonData, FrameData, Output.Pose, Output.Curve);
+	}
+
 	CachedDeltaTime = 0.f; // Reset so that if we evaluate again we don't "create" time inside of the retargeter
 }
 
@@ -48,7 +52,10 @@ void FAnimNode_LiveLinkPose::BuildPoseFromCurveData(const FLiveLinkSubjectFrameD
 	check(BaseStaticData);
 	check(BaseFrameData);
 
-	CurrentRetargetAsset->BuildPoseAndCurveFromBaseData(CachedDeltaTime, BaseStaticData, BaseFrameData, Output.Pose, Output.Curve);
+	if (BaseFrameData->PropertyValues.Num())
+	{
+		CurrentRetargetAsset->BuildPoseAndCurveFromBaseData(CachedDeltaTime, BaseStaticData, BaseFrameData, Output.Pose, Output.Curve);
+	}
 	CachedDeltaTime = 0.f; // Reset so that if we evaluate again we don't "create" time inside of the retargeter
 }
 
